[CSVtoSVG][UI] - Adding a UI wrapper around the CSV to SVG Tool command line tool.

- The arguments are stored as a UObject with reflected properties so I can get a "free" UI by using the property window.  Also, this lets me serialized arguments so I can potentially re-use the settings.  I can also programmatically generate the arugment list to pass to CSVtoSVG by iterating the relfected properties.
- Add SCSVtoSVG widget which contains/owns the arugment UObject, a property window to edit it, and a SStatList.
- When the CSV file property on the arugment changes we infer an output directory and output file name from the input csv file.  These can be overriden.
- I also parse the CSV file to find all the available stats in it to populate a SListView.
- When you click generate to make the SVG graph we automatically open the graph in a web browser.
- Save the current arugments to config when you close the editor.

diff --git a/Engine/Source/Editor/CSVtoSVG/CSVtoSVG.Build.cs b/Engine/Source/Editor/CSVtoSVG/CSVtoSVG.Build.cs
new file mode 100644
--- /dev/null
+++ b/Engine/Source/Editor/CSVtoSVG/CSVtoSVG.Build.cs
@@ -0,0 +1,47 @@
+// Copyright Epic Games, Inc. All Rights Reserved.
+
+using UnrealBuildTool;
+
+public class CSVtoSVG : ModuleRules
+{
+	public CSVtoSVG(ReadOnlyTargetRules Target) : base(Target)
+	{
+		PublicDependencyModuleNames.AddRange(
+			new string[] {
+				"EditorConfig",
+			}
+		);
+
+		PrivateIncludePathModuleNames.AddRange(
+			new string[] {
+				"AssetRegistry",
+                "AssetTools",
+				"EditorWidgets",
+			}
+		);
+
+		PrivateDependencyModuleNames.AddRange(
+			new string[] {
+				"Core",
+				"CoreUObject",
+				"Engine",
+                "InputCore",
+				"Slate",
+				"SlateCore",
+                "EditorStyle",
+				"EditorFramework",
+				"UnrealEd",
+				"ContentBrowserData",
+				"Settings",
+            }
+		);
+
+		DynamicallyLoadedModuleNames.AddRange(
+			new string[] {
+				"AssetRegistry",
+                "AssetTools",
+				"EditorWidgets",
+			}
+		);
+	}
+}
diff --git a/Engine/Source/Editor/CSVtoSVG/Private/CSVtoSVGArguments.cpp b/Engine/Source/Editor/CSVtoSVG/Private/CSVtoSVGArguments.cpp
new file mode 100644
--- /dev/null
+++ b/Engine/Source/Editor/CSVtoSVG/Private/CSVtoSVGArguments.cpp
@@ -0,0 +1,73 @@
+// Copyright Epic Games, Inc. All Rights Reserved.
+
+#include "CSVtoSVGArguments.h"
+#include "CSVtoSVGModule.h"
+
+UCSVtoSVGArugments::UCSVtoSVGArugments(const FObjectInitializer& ObjectInitializer)
+	: Super(ObjectInitializer)
+{
+
+}
+
+FString UCSVtoSVGArugments::GetCommandLine() const
+{
+	if (CSV.FilePath.IsEmpty())
+	{
+		return FString();
+	}
+
+	if (OutputDirectory.Path.IsEmpty())
+	{
+		return FString();
+	}
+
+	if (OutputFilename.IsEmpty())
+	{
+		return FString();
+	}
+
+	// Get the default argument values.
+	UCSVtoSVGArugments* Default = CastChecked<UCSVtoSVGArugments>(GetClass()->GetDefaultObject());
+
+	// Programmatically generate a list of optional arugments by iterating reflected properties.
+	// We use the property name as the argument name to CSVtoSVG.
+	FString Arguments;
+	for (TFieldIterator<FProperty> It(GetClass()); It; ++It)
+	{
+		const FProperty* Property = *It;
+		if (Property && Property->HasMetaData("Category") && (Property->GetMetaData("Category") == TEXT("OptionalArguments")))
+		{
+			const FString ArgumentName = Property->GetName();
+
+			FString ArgumentValue;
+			const bool bDifferentThanDefault = Property->ExportText_InContainer(0, ArgumentValue, this, Default, nullptr, PPF_None);
+
+			// If the user has set a value different from the defaults use the argument.
+			if (bDifferentThanDefault)
+			{
+				// Bool properties are an exception.  They are just commandline switches.
+				if (Property->GetCPPType() == TEXT("bool"))
+				{
+					Arguments += FString::Printf(TEXT(" -%s"), *ArgumentName);
+				}
+				else
+				{
+					Arguments += FString::Printf(TEXT(" -%s %s"), *ArgumentName, *ArgumentValue);
+				}
+			}
+		}
+	}
+
+	// Build an argument string with the required and optional arugments.
+	return FString::Printf(
+		TEXT("-csv %s -o %s %s"),
+		*CSV.FilePath,
+		*GetOutputFileName(),
+		*Arguments
+	);
+}
+
+FString UCSVtoSVGArugments::GetOutputFileName() const
+{
+	return FPaths::Combine(*OutputDirectory.Path, OutputFilename);
+}
\ No newline at end of file
diff --git a/Engine/Source/Editor/CSVtoSVG/Private/CSVtoSVGModule.cpp b/Engine/Source/Editor/CSVtoSVG/Private/CSVtoSVGModule.cpp
new file mode 100644
--- /dev/null
+++ b/Engine/Source/Editor/CSVtoSVG/Private/CSVtoSVGModule.cpp
@@ -0,0 +1,120 @@
+// Copyright Epic Games, Inc. All Rights Reserved.
+
+#include "CSVtoSVGModule.h"
+#include "CSVtoSVGArguments.h"
+#include "Modules/ModuleManager.h"
+#include "Widgets/DeclarativeSyntaxSupport.h"
+#include "Framework/Application/SlateApplication.h"
+#include "Textures/SlateIcon.h"
+#include "Framework/Docking/TabManager.h"
+#include "EditorStyleSet.h"
+#include "Editor/WorkspaceMenuStructure/Public/WorkspaceMenuStructure.h"
+#include "Editor/WorkspaceMenuStructure/Public/WorkspaceMenuStructureModule.h"
+#include "Widgets/Docking/SDockTab.h"
+#include "Modules/ModuleManager.h"
+#include "Widgets/Docking/SDockTab.h"
+#include "SCSVtoSVG.h"
+#include "ISettingsModule.h"
+#include "HAL/PlatformProcess.h"
+#include "Misc/Paths.h"
+
+#define LOCTEXT_NAMESPACE "CSVtoSVG"
+
+IMPLEMENT_MODULE(FCSVtoSVGModule, CSVtoSVG);
+DEFINE_LOG_CATEGORY(LogCSVtoSVG);
+
+namespace FCSVtoSVG
+{
+	static const FName CSVtoSVGApp = FName(TEXT("CSVtoSVGApp"));
+}
+
+TSharedRef<SDockTab> FCSVtoSVGModule::CreateTab(const FSpawnTabArgs& Args)
+{
+	TSharedPtr<SDockTab> MajorTab;
+	SAssignNew(MajorTab, SDockTab)
+		.TabRole(ETabRole::MajorTab);
+
+	MajorTab->SetContent(SNew(SCSVtoSVG));
+
+	return MajorTab.ToSharedRef();
+}
+
+
+void FCSVtoSVGModule::StartupModule()
+{
+	FGlobalTabmanager::Get()->RegisterNomadTabSpawner(FCSVtoSVG::CSVtoSVGApp, FOnSpawnTab::CreateRaw(this, &FCSVtoSVGModule::CreateTab))
+		.SetDisplayName(NSLOCTEXT("CSVtoSVGApp", "TabTitle", "CSV to SVG"))
+		.SetTooltipText(NSLOCTEXT("CSVtoSVGApp", "TooltipText", "Tool for generating vector line graphs from comma-separated value files generated from CSV profiles."))
+		.SetGroup(WorkspaceMenu::GetMenuStructure().GetToolsCategory())
+		.SetIcon(FSlateIcon(FEditorStyle::GetStyleSetName(), "ClassIcon.UserDefinedStruct"));
+}
+
+void FCSVtoSVGModule::ShutdownModule()
+{
+	if (FSlateApplication::IsInitialized())
+	{
+		FGlobalTabmanager::Get()->UnregisterNomadTabSpawner(FCSVtoSVG::CSVtoSVGApp);
+	}
+}
+
+FFilePath GenerateSVG(const UCSVtoSVGArugments& Arguments, const TArray<FString>& StatList)
+{
+	const FString CSVtoSVGExe = FPaths::Combine(FPaths::EngineDir(), TEXT("Binaries/DotNET/CsvTools"), TEXT("CSVToSVG.exe"));
+
+	FString CmdLine = Arguments.GetCommandLine();
+	if (CmdLine.IsEmpty())
+	{
+		return FFilePath();
+	}
+
+	// Add selected stats to the command line.
+	CmdLine += FString(TEXT(" -stats"));
+	for (const FString& Stat : StatList)
+	{
+		CmdLine += TEXT(" ") + Stat;
+	}
+
+	UE_LOG(LogCSVtoSVG, Display, TEXT("Generating SVG: %s"), *CmdLine);
+
+	// Create a read and write pipe for the child process
+	void* PipeRead = nullptr;
+	void* PipeWrite = nullptr;
+	if (!FPlatformProcess::CreatePipe(PipeRead, PipeWrite))
+	{
+		return FFilePath();
+	}
+
+	FProcHandle ProcHandle = FPlatformProcess::CreateProc(*CSVtoSVGExe, *CmdLine, true, false, false, nullptr, 0, nullptr, PipeWrite, PipeRead);
+	if (ProcHandle.IsValid())
+	{
+		FString CSVtoCSVOutput;
+		while (FPlatformProcess::IsProcRunning(ProcHandle))
+		{
+			CSVtoCSVOutput += FPlatformProcess::ReadPipe(PipeRead);
+
+		}
+		FPlatformProcess::WaitForProc(ProcHandle);
+
+		if (!CSVtoCSVOutput.IsEmpty())
+		{
+			UE_LOG(LogCSVtoSVG, Display, TEXT("%s"), *CSVtoCSVOutput);
+		}
+
+		if (!CSVtoCSVOutput.Contains(TEXT("[ERROR]")))
+		{
+			// If there was no error open the output file.
+			const FString OutputFileName = Arguments.GetOutputFileName();
+			if (!OutputFileName.IsEmpty())
+			{
+				FPlatformProcess::LaunchURL(*OutputFileName, nullptr, nullptr);
+			}
+		}
+
+		FPlatformProcess::ClosePipe(PipeRead, PipeWrite);
+		FPlatformProcess::CloseProc(ProcHandle);
+	}
+
+	return FFilePath();
+}
+
+#undef LOCTEXT_NAMESPACE
diff --git a/Engine/Source/Editor/CSVtoSVG/Private/SCSVtoSVG.cpp b/Engine/Source/Editor/CSVtoSVG/Private/SCSVtoSVG.cpp
new file mode 100644
--- /dev/null
+++ b/Engine/Source/Editor/CSVtoSVG/Private/SCSVtoSVG.cpp
@@ -0,0 +1,139 @@
+// Copyright Epic Games, Inc. All Rights Reserved.
+
+#include "SCSVtoSVG.h"
+#include "CSVtoSVGModule.h"
+#include "SStatList.h"
+#include "CSVtoSVGArguments.h"
+#include "Widgets/Input/SButton.h"
+#include "Widgets/Layout/SScrollBox.h"
+#include "Misc/FileHelper.h"
+#include "Misc/Paths.h"
+
+#define LOCTEXT_NAMESPACE "CSVtoSVG"
+
+void SCSVtoSVG::Construct(const FArguments& Args)
+{
+	FPropertyEditorModule& PropertyEditorModule = FModuleManager::Get().GetModuleChecked<FPropertyEditorModule>("PropertyEditor");
+
+	FDetailsViewArgs DetailsViewArgs;
+	DetailsViewArgs.NameAreaSettings = FDetailsViewArgs::HideNameArea;
+	DetailsViewArgs.bAllowSearch = false;
+	DetailsViewArgs.NotifyHook = this;
+	DetailsViewArgs.bShowOptions = true;
+	DetailsViewArgs.bShowModifiedPropertiesOption = true;
+
+	Arguments = TStrongObjectPtr<UCSVtoSVGArugments>(NewObject<UCSVtoSVGArugments>());
+	Arguments->LoadEditorConfig();
+
+	CSVtoSVGArgumentsDetailsView = PropertyEditorModule.CreateDetailView(DetailsViewArgs);
+	CSVtoSVGArgumentsDetailsView->SetObject(Arguments.Get());
+
+	const FMargin PaddingAmount(5.0f);
+
+	this->ChildSlot
+	[
+		SNew(SVerticalBox)
+		+ SVerticalBox::Slot()
+		.Padding(PaddingAmount)
+		[
+			SNew(SHorizontalBox)
+
+			+ SHorizontalBox::Slot()
+			.Padding(PaddingAmount)
+			[
+				SNew(SScrollBox)
+				+ SScrollBox::Slot()
+				.Padding(PaddingAmount)
+				[
+					CSVtoSVGArgumentsDetailsView->AsShared()
+				]
+			]
+			+ SHorizontalBox::Slot()
+			.Padding(PaddingAmount)
+			[
+				SAssignNew(StatListView, SStatList)
+			]
+		]
+		+ SVerticalBox::Slot()
+			.Padding(PaddingAmount)
+			.AutoHeight()
+			.HAlign(HAlign_Right)
+		[
+			SNew(SButton)
+			.ContentPadding(FEditorStyle::GetMargin("StandardDialog.ContentPadding"))
+			.ButtonStyle(FEditorStyle::Get(), "FlatButton.Success")
+			.ForegroundColor(FLinearColor::White)
+			.OnClicked(FOnClicked::CreateSP(this, &SCSVtoSVG::OnGenerateGraphClicked))
+			.ToolTipText(LOCTEXT("GenerateSVGGraphTooltipLoc", "Generates an SVG graph."))
+			.Text(LOCTEXT("GenerateSVGGraphLoc", "Generate"))
+		]
+	];
+
+	// Load CSV, which in turn will populate the UI.
+	if (!Arguments->CSV.FilePath.IsEmpty())
+	{
+		LoadCSVFile(*Arguments->CSV.FilePath);
+	}
+}
+
+/** SCSVtoSVG destructor */
+SCSVtoSVG::~SCSVtoSVG()
+{
+	Arguments->SaveEditorConfig();
+}
+
+FReply SCSVtoSVG::OnGenerateGraphClicked()
+{
+	FFilePath OutputFilePath = GenerateSVG(*Arguments, StatListView->GetSelectedStats());
+	if (!OutputFilePath.FilePath.IsEmpty())
+	{
+		return FReply::Handled();
+	}
+
+	return FReply::Unhandled();
+}
+
+void SCSVtoSVG::StatListSelectionChanged(const TArray<FString>& Stats)
+{
+	StatListView->UpdateStatList(Stats);
+}
+
+void SCSVtoSVG::NotifyPostChange(const FPropertyChangedEvent& PropertyChangedEvent, FProperty* PropertyThatChanged)
+{
+	if (PropertyChangedEvent.MemberProperty && PropertyChangedEvent.MemberProperty->GetName() == TEXT("CSV"))
+	{
+		LoadCSVFile(*Arguments->CSV.FilePath);
+	}
+}
+
+void SCSVtoSVG::LoadCSVFile(const FString& Filepath)
+{
+	TArray<FString> Contents;
+	if (FFileHelper::LoadFileToStringArray(Contents, *Filepath))
+	{
+		if (Contents.Num())
+		{
+			const FString& StatLine = Contents[0];
+			TArray<FString> Stats;
+			StatLine.ParseIntoArray(Stats, TEXT(","));
+
+			Stats.Sort(TLess<FString>());
+
+			StatListSelectionChanged(Stats);
+		}
+
+		// If the user has not specified a value for filename infer one from the source csv file.
+		if (Arguments->OutputFilename.IsEmpty())
+		{
+			Arguments->OutputFilename = FPaths::ChangeExtension(FPaths::GetCleanFilename(Arguments->CSV.FilePath), TEXT("svg"));
+		}
+
+		// If the user has not specified a value for directory infer ones from the source csv file directory
+		if (Arguments->OutputDirectory.Path.IsEmpty())
+		{
+			Arguments->OutputDirectory.Path = FPaths::ConvertRelativePathToFull(FPaths::GetPath(Arguments->CSV.FilePath));
+		}
+	}
+}
+
+#undef LOCTEXT_NAMESPACE
diff --git a/Engine/Source/Editor/CSVtoSVG/Private/SCSVtoSVG.h b/Engine/Source/Editor/CSVtoSVG/Private/SCSVtoSVG.h
new file mode 100644
--- /dev/null
+++ b/Engine/Source/Editor/CSVtoSVG/Private/SCSVtoSVG.h
@@ -0,0 +1,46 @@
+// Copyright Epic Games, Inc. All Rights Reserved.
+
+#pragma once
+
+#include "CoreMinimal.h"
+#include "Widgets/SCompoundWidget.h"
+#include "Misc/NotifyHook.h"
+
+class UCSVtoSVGArugments;
+class IDetailsView;
+class SStatList;
+
+class SCSVtoSVG : public SCompoundWidget, public FNotifyHook
+{
+	SLATE_BEGIN_ARGS(SCSVtoSVG)
+	{
+	}
+
+	SLATE_END_ARGS()
+
+	void Construct(const FArguments& Args);
+
+	/** SCSVtoSVG destructor */
+	~SCSVtoSVG();
+
+private:
+
+	FReply OnGenerateGraphClicked();
+
+	void StatListSelectionChanged(const TArray<FString>& Stats);
+
+	//~ Begin FNotifyHook Interface
+	virtual void NotifyPostChange(const FPropertyChangedEvent& PropertyChangedEvent, FProperty* PropertyThatChanged) override;
+	// End of FNotifyHook
+
+	void LoadCSVFile(const FString& Filepath);
+
+	/** This is the Detail View for the currently selected UActorRecording or an item from an Extender */
+	TSharedPtr<IDetailsView> CSVtoSVGArgumentsDetailsView = nullptr;
+
+	/** The list of stats parsed from the CSV file. */
+	TSharedPtr<SStatList> StatListView = nullptr;
+
+	/** The current arguments. */
+	TStrongObjectPtr<UCSVtoSVGArugments> Arguments = nullptr;
+};
diff --git a/Engine/Source/Editor/CSVtoSVG/Private/SStatList.cpp b/Engine/Source/Editor/CSVtoSVG/Private/SStatList.cpp
new file mode 100644
--- /dev/null
+++ b/Engine/Source/Editor/CSVtoSVG/Private/SStatList.cpp
@@ -0,0 +1,63 @@
+// Copyright Epic Games, Inc. All Rights Reserved.
+
+#include "SStatList.h"
+#include "Widgets/Views/SListView.h"
+
+#define LOCTEXT_NAMESPACE "CSVtoSVG"
+
+void SStatList::Construct(const FArguments& Args)
+{
+	StatListView = SNew(SListView<TSharedPtr<FString>>)
+					.ItemHeight(20.0f)
+					.ListItemsSource(&StatList)
+					.OnGenerateRow(this, &SStatList::OnGenerateWidgetForList);
+
+	ChildSlot
+	[
+		SNew(SHorizontalBox)
+		+ SHorizontalBox::Slot()
+		[
+			StatListView.ToSharedRef()
+		]
+	];
+}
+
+TSharedRef<ITableRow> SStatList::OnGenerateWidgetForList(TSharedPtr<FString> InItem, const TSharedRef<STableViewBase>& OwnerTable)
+{
+	return
+		SNew(STableRow<TSharedPtr<FString>>, OwnerTable)
+		[
+			SNew(STextBlock)
+				.Text(FText::FromString(*InItem.Get()))
+		];
+}
+
+void SStatList::UpdateStatList(const TArray<FString>& StatNames)
+{
+	StatList.Empty();
+
+	for (const auto& StatName : StatNames)
+	{
+		StatList.Add(MakeShared<FString>(StatName));
+	}
+
+	StatListView->RequestListRefresh();
+}
+
+TArray<FString> SStatList::GetSelectedStats() const
+{
+	TArray<TSharedPtr<FString>> SelectedItems = StatListView->GetSelectedItems();
+
+	TArray<FString> SelectedStats;
+	for (const TSharedPtr<FString>& Item : SelectedItems)
+	{
+		if (Item.IsValid())
+		{
+			SelectedStats.Add(*Item);
+		}
+	}
+
+	return SelectedStats;
+}
+
+#undef LOCTEXT_NAMESPACE
diff --git a/Engine/Source/Editor/CSVtoSVG/Private/SStatList.h b/Engine/Source/Editor/CSVtoSVG/Private/SStatList.h
new file mode 100644
--- /dev/null
+++ b/Engine/Source/Editor/CSVtoSVG/Private/SStatList.h
@@ -0,0 +1,28 @@
+// Copyright Epic Games, Inc. All Rights Reserved.
+
+#pragma once
+
+#include "CoreMinimal.h"
+#include "Widgets/SCompoundWidget.h"
+
+class SStatList : public SCompoundWidget
+{
+	SLATE_BEGIN_ARGS(SStatList)
+	{
+	}
+
+	SLATE_END_ARGS()
+
+	void Construct(const FArguments& Args);
+
+	void UpdateStatList(const TArray<FString>& StatNames);
+
+	TArray<FString> GetSelectedStats() const;
+
+private:
+
+	TSharedRef<ITableRow> OnGenerateWidgetForList(TSharedPtr<FString> InItem, const TSharedRef<STableViewBase>& OwnerTable);
+
+	TSharedPtr<SListView<TSharedPtr<FString>>> StatListView;
+	TArray<TSharedPtr<FString>> StatList;
+};
diff --git a/Engine/Source/Editor/CSVtoSVG/Public/CSVtoSVGArguments.h b/Engine/Source/Editor/CSVtoSVG/Public/CSVtoSVGArguments.h
new file mode 100644
--- /dev/null
+++ b/Engine/Source/Editor/CSVtoSVG/Public/CSVtoSVGArguments.h
@@ -0,0 +1,93 @@
+// Copyright Epic Games, Inc. All Rights Reserved.
+
+#pragma once
+
+#include "CoreMinimal.h"
+#include "EditorConfigBase.h"
+#include "Engine/EngineTypes.h"
+#include "CSVtoSVGArguments.generated.h"
+
+DECLARE_MULTICAST_DELEGATE_TwoParams(FCSVtoSVGArugmentsPropertySetModifiedSignature, UObject*, FProperty*);
+
+UCLASS(EditorConfig = "CSVtoSVG")
+class UCSVtoSVGArugments : public UEditorConfigBase
+{
+	GENERATED_UCLASS_BODY()
+
+	/** The base directories to be considered Internal Only for the struct picker.*/
+	UPROPERTY(EditAnywhere, Category = RequiredArguments, meta = (EditorConfig, DisplayName = "CSV", Tooltip = "Specifies a list of CSVs separated by spaces"))
+	FFilePath CSV;
+
+	UPROPERTY(EditAnywhere, Category = RequiredArguments, meta = (EditorConfig, DisplayName = "Ouput Directory", Tooltip = "Sets the filename for the SVG output."))
+	FDirectoryPath OutputDirectory;
+
+	UPROPERTY(EditAnywhere, Category = RequiredArguments, meta = (EditorConfig, DisplayName = "Ouput Filename", Tooltip = "Sets the filename for the SVG output."))
+	FString OutputFilename;
+
+	UPROPERTY(EditAnywhere, Category = OptionalArguments, meta = (EditorConfig, DisplayName = "Skip Rows", Tooltip = "Skips a specified number of rows in the CSV. This is useful for CSV file sgenerated from the FPSChartStart command, where there's a 4-row summary at the top."))
+	int32 skipRows = 0;
+
+	UPROPERTY(EditAnywhere, Category = OptionalArguments, meta = (EditorConfig, DisplayName = "minX", Tooltip = "Clamps the X and Y range of the source data."))
+	float minX = -100000.0f;
+
+	UPROPERTY(EditAnywhere, Category = OptionalArguments, meta = (EditorConfig, DisplayName = "maxX ", Tooltip = "Clamps the X and Y range of the source data."))
+	float maxX = -100000.0f;
+
+	UPROPERTY(EditAnywhere, Category = OptionalArguments, meta = (EditorConfig, DisplayName = "minY", Tooltip = "Clamps the X and Y range of the source data."))
+	float minY = -100000.0f;
+
+	UPROPERTY(EditAnywhere, Category = OptionalArguments, meta = (EditorConfig, DisplayName = "maxY", Tooltip = "Clamps the X and Y range of the source data."))
+	float maxY = -100000.0f;
+
+	UPROPERTY(EditAnywhere, Category = OptionalArguments, meta = (EditorConfig, DisplayName = "Smooth", Tooltip = "Smooth the graph."))
+	bool smooth = false;
+
+	UPROPERTY(EditAnywhere, Category = OptionalArguments, meta = (EditorConfig, DisplayName = "Smooth kernel size", Tooltip = "Specifies the smoothing kernel size in column entries to use."))
+	int32 smoothKernelSize = -1;
+
+	UPROPERTY(EditAnywhere, Category = OptionalArguments, meta = (EditorConfig, DisplayName = "Smooth kernel size", Tooltip = "Specifies the smoothing kernel size as a percentage of the graph length."))
+	float smoothKernelPercent = -1.0f;
+
+	UPROPERTY(EditAnywhere, Category = OptionalArguments, meta = (EditorConfig, DisplayName = "Width", Tooltip = "The width of the graph."))
+	int32 width = 1000;
+
+	UPROPERTY(EditAnywhere, Category = OptionalArguments, meta = (EditorConfig, DisplayName = "Height", Tooltip = "The height of the graph."))
+	int32 height = 500;
+
+	UPROPERTY(EditAnywhere, Category = OptionalArguments, meta = (EditorConfig, DisplayName = "Title", Tooltip = "Sets the title for the graph."))
+	FString title;
+
+	UPROPERTY(EditAnywhere, Category = OptionalArguments, meta = (EditorConfig, DisplayName = "No Metadata", Tooltip = ""))
+	bool noMetadata = false;
+
+	UPROPERTY(EditAnywhere, Category = OptionalArguments, meta = (EditorConfig, DisplayName = "Graph only", Tooltip = "Render just the graph, no borders."))
+	bool graphOnly = false;
+
+	UPROPERTY(EditAnywhere, Category = OptionalArguments, meta = (EditorConfig, DisplayName = "Budget [ms]", Tooltip = "Sets the budget line. Default is 33.3."))
+	float budget = 33.3f;
+
+	UPROPERTY(EditAnywhere, Category = OptionalArguments, meta = (EditorConfig, DisplayName = "Thickness", Tooltip = "Sets the line thickness of the graph."))
+	float thickness = 1.0f;
+
+	UPROPERTY(EditAnywhere, Category = OptionalArguments, meta = (EditorConfig, DisplayName = "Stacked", Tooltip = "Makes a stacked graph for cumulative stats."))
+	bool stacked = false;
+
+	UPROPERTY(EditAnywhere, Category = OptionalArguments, meta = (EditorConfig, DisplayName = "Interactive", Tooltip = "Adds an interactive frame marker."))
+	bool interactive = false;
+
+	UPROPERTY(EditAnywhere, Category = OptionalArguments, meta = (EditorConfig, DisplayName = "Show Averages", Tooltip = "Shows stat averages next to the legend and sorts the stats in the legend high to low by average value."))
+	bool showaverages = false;
+
+	FString GetCommandLine() const;
+
+	FString GetOutputFileName() const;
+
+	/** @return the multicast delegate that is called when properties are modified */
+	FCSVtoSVGArugmentsPropertySetModifiedSignature& GetOnModified()
+	{
+		return OnModified;
+	}
+
+private:
+	FCSVtoSVGArugmentsPropertySetModifiedSignature OnModified;
+};
diff --git a/Engine/Source/Editor/CSVtoSVG/Public/CSVtoSVGModule.h b/Engine/Source/Editor/CSVtoSVG/Public/CSVtoSVGModule.h
new file mode 100644
--- /dev/null
+++ b/Engine/Source/Editor/CSVtoSVG/Public/CSVtoSVGModule.h
@@ -0,0 +1,40 @@
+// Copyright Epic Games, Inc. All Rights Reserved.
+
+#pragma once
+
+#include "CoreMinimal.h"
+#include "Modules/ModuleInterface.h"
+
+class SDockTab;
+class UCSVtoSVGArugments;
+
+DECLARE_LOG_CATEGORY_EXTERN(LogCSVtoSVG, Log, All);
+
+/**
+ * Generates a SVG graph from the given arguments.
+ *
+ * @param	Arguments								The arguments to define and configure graph generation.
+ *
+ * @return	File path to the generated SVG file.
+ */
+FFilePath GenerateSVG(const UCSVtoSVGArugments& Arguments, const TArray<FString>& StatList);
+
+/**
+ * Struct Viewer module
+ */
+class FCSVtoSVGModule : public IModuleInterface
+{
+public:
+	/**
+	 * Called right after the module DLL has been loaded and the module object has been created
+	 */
+	virtual void StartupModule();
+
+	/**
+	 * Called before the module is unloaded, right before the module object is destroyed.
+	 */
+	virtual void ShutdownModule();
+
+private:
+	TSharedRef<SDockTab> CreateTab(const FSpawnTabArgs& Args);
+};
diff --git a/Engine/Source/Editor/UnrealEd/Private/EditorEngine.cpp b/Engine/Source/Editor/UnrealEd/Private/EditorEngine.cpp
--- a/Engine/Source/Editor/UnrealEd/Private/EditorEngine.cpp
+++ b/Engine/Source/Editor/UnrealEd/Private/EditorEngine.cpp
@@ -1045,7 +1045,8 @@ void UEditorEngine::Init(IEngineLoop* InEngineLoop)
 			TEXT("EditorFramework"),
 			TEXT("WorldPartitionEditor"),
 			TEXT("EditorConfig"),
-			TEXT("DerivedDataEditor")
+			TEXT("DerivedDataEditor"),
+			TEXT("CSVtoSVG"),
 		};
 
 		FScopedSlowTask ModuleSlowTask(UE_ARRAY_COUNT(ModuleNames));
diff --git a/Engine/Source/Editor/UnrealEd/UnrealEd.Build.cs b/Engine/Source/Editor/UnrealEd/UnrealEd.Build.cs
--- a/Engine/Source/Editor/UnrealEd/UnrealEd.Build.cs
+++ b/Engine/Source/Editor/UnrealEd/UnrealEd.Build.cs
@@ -70,6 +70,7 @@ public class UnrealEd : ModuleRules
 				"VirtualTexturingEditor",
 				"TextureBuild",
 				"ToolWidgets",
+				"CSVtoSVG"
 			}
 		);
 
@@ -266,6 +267,7 @@ public class UnrealEd : ModuleRules
 				"ClothPainter",
 				"Media",
 				"VirtualTexturingEditor",
+				"CSVtoSVG",
 			}
 		);
 
