[BuildFix][NonUnity] - Missing headers and forward declarations.

diff --git a/Engine/Source/Editor/CSVtoSVG/Private/CSVtoSVGArguments.cpp b/Engine/Source/Editor/CSVtoSVG/Private/CSVtoSVGArguments.cpp
--- a/Engine/Source/Editor/CSVtoSVG/Private/CSVtoSVGArguments.cpp
+++ b/Engine/Source/Editor/CSVtoSVG/Private/CSVtoSVGArguments.cpp
@@ -2,6 +2,7 @@
 
 #include "CSVtoSVGArguments.h"
 #include "CSVtoSVGModule.h"
+#include "Misc/Paths.h"
 
 UCSVtoSVGArugments::UCSVtoSVGArugments(const FObjectInitializer& ObjectInitializer)
 	: Super(ObjectInitializer)
diff --git a/Engine/Source/Editor/CSVtoSVG/Private/SCSVtoSVG.cpp b/Engine/Source/Editor/CSVtoSVG/Private/SCSVtoSVG.cpp
--- a/Engine/Source/Editor/CSVtoSVG/Private/SCSVtoSVG.cpp
+++ b/Engine/Source/Editor/CSVtoSVG/Private/SCSVtoSVG.cpp
@@ -4,10 +4,13 @@
 #include "CSVtoSVGModule.h"
 #include "SStatList.h"
 #include "CSVtoSVGArguments.h"
+#include "DetailsViewArgs.h"
+#include "Modules/ModuleManager.h"
 #include "Widgets/Input/SButton.h"
 #include "Widgets/Layout/SScrollBox.h"
 #include "Misc/FileHelper.h"
 #include "Misc/Paths.h"
+#include "PropertyEditorModule.h"
 
 #define LOCTEXT_NAMESPACE "CSVtoSVG"
 
diff --git a/Engine/Source/Editor/CSVtoSVG/Private/SCSVtoSVG.h b/Engine/Source/Editor/CSVtoSVG/Private/SCSVtoSVG.h
--- a/Engine/Source/Editor/CSVtoSVG/Private/SCSVtoSVG.h
+++ b/Engine/Source/Editor/CSVtoSVG/Private/SCSVtoSVG.h
@@ -3,6 +3,7 @@
 #pragma once
 
 #include "CoreMinimal.h"
+#include "UObject/StrongObjectPtr.h"
 #include "Widgets/SCompoundWidget.h"
 #include "Misc/NotifyHook.h"
 
diff --git a/Engine/Source/Editor/CSVtoSVG/Private/SStatList.h b/Engine/Source/Editor/CSVtoSVG/Private/SStatList.h
--- a/Engine/Source/Editor/CSVtoSVG/Private/SStatList.h
+++ b/Engine/Source/Editor/CSVtoSVG/Private/SStatList.h
@@ -3,8 +3,12 @@
 #pragma once
 
 #include "CoreMinimal.h"
+#include "Widgets/Views/SListView.h"
 #include "Widgets/SCompoundWidget.h"
 
+class ITableRow;
+class STableViewBase;
+
 class SStatList : public SCompoundWidget
 {
 	SLATE_BEGIN_ARGS(SStatList)
diff --git a/Engine/Source/Editor/CSVtoSVG/Public/CSVtoSVGModule.h b/Engine/Source/Editor/CSVtoSVG/Public/CSVtoSVGModule.h
--- a/Engine/Source/Editor/CSVtoSVG/Public/CSVtoSVGModule.h
+++ b/Engine/Source/Editor/CSVtoSVG/Public/CSVtoSVGModule.h
@@ -3,10 +3,12 @@
 #pragma once
 
 #include "CoreMinimal.h"
+#include "UObject/SoftObjectPath.h"
 #include "Modules/ModuleInterface.h"
 
 class SDockTab;
 class UCSVtoSVGArugments;
+class FSpawnTabArgs;
 
 DECLARE_LOG_CATEGORY_EXTERN(LogCSVtoSVG, Log, All);
 
